removed explicit boxing of TWeakObjPtr, because this is no longer necessary since the TWeakObjPtr assignment operators have been fixed

diff --git a/Engine/Plugins/Messaging/UdpMessaging/Source/UdpMessaging/Private/UdpMessagingModule.cpp b/Engine/Plugins/Messaging/UdpMessaging/Source/UdpMessaging/Private/UdpMessagingModule.cpp
--- a/Engine/Plugins/Messaging/UdpMessaging/Source/UdpMessaging/Private/UdpMessagingModule.cpp
+++ b/Engine/Plugins/Messaging/UdpMessaging/Source/UdpMessaging/Private/UdpMessagingModule.cpp
@@ -46,7 +46,7 @@ public:
 			SettingsModule->RegisterSettings("Project", "Plugins", "UdpMessaging",
 				LOCTEXT("UdpMessagingSettingsName", "UDP Messaging"),
 				LOCTEXT("UdpMessagingSettingsDescription", "Configure the UDP Messaging plug-in."),
-				TWeakObjectPtr<UObject>(UUdpMessagingSettings::StaticClass()->GetDefaultObject()),
+				GetMutableDefault<UUdpMessagingSettings>(),
 				SettingsDelegates
 			);
 		}
diff --git a/Engine/Source/Developer/Android/AndroidPlatformEditor/Private/AndroidPlatformEditorModule.cpp b/Engine/Source/Developer/Android/AndroidPlatformEditor/Private/AndroidPlatformEditorModule.cpp
--- a/Engine/Source/Developer/Android/AndroidPlatformEditor/Private/AndroidPlatformEditorModule.cpp
+++ b/Engine/Source/Developer/Android/AndroidPlatformEditor/Private/AndroidPlatformEditorModule.cpp
@@ -28,8 +28,8 @@ class FAndroidPlatformEditorModule : public IModuleInterface
 			SettingsModule->RegisterSettings("Project", "Platforms", "Android",
 				LOCTEXT("RuntimeSettingsName", "Android"),
 				LOCTEXT("RuntimeSettingsDescription", "Settings and resources for Android platforms"),
-				TWeakObjectPtr<UObject>(GetMutableDefault<UAndroidRuntimeSettings>())
-				);
+				GetMutableDefault<UAndroidRuntimeSettings>()
+			);
 		}
 	}
 
diff --git a/Engine/Source/Developer/iOS/IOSPlatformEditor/Private/IOSPlatformEditorModule.cpp b/Engine/Source/Developer/iOS/IOSPlatformEditor/Private/IOSPlatformEditorModule.cpp
--- a/Engine/Source/Developer/iOS/IOSPlatformEditor/Private/IOSPlatformEditorModule.cpp
+++ b/Engine/Source/Developer/iOS/IOSPlatformEditor/Private/IOSPlatformEditorModule.cpp
@@ -28,8 +28,8 @@ class FIOSPlatformEditorModule : public IModuleInterface
 			SettingsModule->RegisterSettings("Project", "Platforms", "iOS",
 				LOCTEXT("RuntimeSettingsName", "iOS"),
 				LOCTEXT("RuntimeSettingsDescription", "Settings and resources for the iOS platform"),
-				TWeakObjectPtr<UObject>(GetMutableDefault<UIOSRuntimeSettings>())
-				);
+				GetMutableDefault<UIOSRuntimeSettings>()
+			);
 		}
 	}
 
diff --git a/Engine/Source/Editor/DetailCustomizations/Private/DialogueWaveWidgets.cpp b/Engine/Source/Editor/DetailCustomizations/Private/DialogueWaveWidgets.cpp
--- a/Engine/Source/Editor/DetailCustomizations/Private/DialogueWaveWidgets.cpp
+++ b/Engine/Source/Editor/DetailCustomizations/Private/DialogueWaveWidgets.cpp
@@ -325,7 +325,7 @@ void SDialogueVoicePropertyEditor::OnBrowseToDialogueVoice()
 	if( DialogueVoice )
 	{	
 		// Find the item in the content browser
-		GoToAssetInContentBrowser( TWeakObjectPtr<UObject>( DialogueVoice ) );
+		GoToAssetInContentBrowser(DialogueVoice);
 	}
 }
 
diff --git a/Engine/Source/Editor/EditorSettingsViewer/Private/EditorSettingsViewerModule.cpp b/Engine/Source/Editor/EditorSettingsViewer/Private/EditorSettingsViewerModule.cpp
--- a/Engine/Source/Editor/EditorSettingsViewer/Private/EditorSettingsViewerModule.cpp
+++ b/Engine/Source/Editor/EditorSettingsViewer/Private/EditorSettingsViewerModule.cpp
@@ -95,7 +95,7 @@ protected:
 		SettingsModule.RegisterSettings("Editor", "General", "AutomationTest",
 			LOCTEXT("AutomationSettingsName", "Automation"),
 			LOCTEXT("AutomationSettingsDescription", "Set up automation test assets."),
-			TWeakObjectPtr<UObject>(GetMutableDefault<UAutomationTestSettings>())
+			GetMutableDefault<UAutomationTestSettings>()
 		);
 
 		// input bindings
@@ -115,9 +115,9 @@ protected:
 		SettingsModule.RegisterSettings("Editor", "General", "Internationalization",
 			LOCTEXT("InternationalizationSettingsModelName", "Region & Language"),
 			LOCTEXT("InternationalizationSettingsModelDescription", "Configure the editor's behavior to use a language and fit a region's culture."),
-			TWeakObjectPtr<UObject>(GetMutableDefault<UInternationalizationSettingsModel>()),
+			GetMutableDefault<UInternationalizationSettingsModel>(),
 			RegionAndLanguageDelegates
-			);
+		);
 
 		SettingsModule.RegisterSettings("Editor", "General", "InputBindings",
 			LOCTEXT("InputBindingsSettingsName", "Keyboard Shortcuts"),
@@ -130,7 +130,7 @@ protected:
 		SettingsModule.RegisterSettings("Editor", "General", "LoadingSaving",
 			LOCTEXT("LoadingSavingSettingsName", "Loading & Saving"),
 			LOCTEXT("LoadingSavingSettingsDescription", "Change how the Editor loads and saves files."),
-			TWeakObjectPtr<UObject>(GetMutableDefault<UEditorLoadingSavingSettings>()),
+			GetMutableDefault<UEditorLoadingSavingSettings>(),
 			UserSettingsDelegates
 		);
 
@@ -141,7 +141,7 @@ protected:
 		SettingsModule.RegisterSettings("Editor", "General", "UserSettings",
 			LOCTEXT("UserSettingsName", "Miscellaneous"),
 			LOCTEXT("UserSettingsDescription", "Customize the behavior, look and feel of the editor."),
-			TWeakObjectPtr<UObject>(&GEditor->AccessEditorUserSettings()),
+			&GEditor->AccessEditorUserSettings(),
 			UserSettingsDelegates
 		);
 
@@ -149,7 +149,7 @@ protected:
 		SettingsModule.RegisterSettings("Editor", "General", "Experimental",
 			LOCTEXT("ExperimentalettingsName", "Experimental"),
 			LOCTEXT("ExperimentalSettingsDescription", "Enable and configure experimental Editor features."),
-			TWeakObjectPtr<UObject>(GetMutableDefault<UEditorExperimentalSettings>())
+			GetMutableDefault<UEditorExperimentalSettings>()
 		);
 	}
 
@@ -162,21 +162,21 @@ protected:
 		SettingsModule.RegisterSettings("Editor", "LevelEditor", "PlayIn",
 			LOCTEXT("LevelEditorPlaySettingsName", "Play"),
 			LOCTEXT("LevelEditorPlaySettingsDescription", "Set up window sizes and other options for the Play In Editor (PIE) feature."),
-			TWeakObjectPtr<UObject>(GetMutableDefault<ULevelEditorPlaySettings>())
+			GetMutableDefault<ULevelEditorPlaySettings>()
 		);
 
 		// view port settings
 		SettingsModule.RegisterSettings("Editor", "LevelEditor", "Viewport",
 			LOCTEXT("LevelEditorViewportSettingsName", "Viewports"),
 			LOCTEXT("LevelEditorViewportSettingsDescription", "Configure the look and feel of the Level Editor view ports."),
-			TWeakObjectPtr<UObject>(GetMutableDefault<ULevelEditorViewportSettings>())
+			GetMutableDefault<ULevelEditorViewportSettings>()
 		);
 
 		// miscellaneous settings
 		SettingsModule.RegisterSettings("Editor", "LevelEditor", "Misc",
 			LOCTEXT("LevelEditorMiscSettingsName", "Miscellaneous"),
 			LOCTEXT("LevelEditorMiscSettingsDescription", "Configure miscellaneous settings for the Level Editor."),
-			TWeakObjectPtr<UObject>(GetMutableDefault<ULevelEditorMiscSettings>())
+			GetMutableDefault<ULevelEditorMiscSettings>()
 		);
 	}
 
@@ -189,14 +189,14 @@ protected:
 		SettingsModule.RegisterSettings("Editor", "ContentEditors", "ContentBrowser",
 			LOCTEXT("ContentEditorsContentBrowserSettingsName", "Content Browser"),
 			LOCTEXT("ContentEditorsContentBrowserSettingsDescription", "Change the behavior of the Content Browser."),
-			TWeakObjectPtr<UObject>(GetMutableDefault<UContentBrowserSettings>())
+			GetMutableDefault<UContentBrowserSettings>()
 		);
 
 		// destructable mesh editor
 /*		SettingsModule.RegisterSettings("Editor", "ContentEditors", "DestructableMeshEditor",
 			LOCTEXT("ContentEditorsDestructableMeshEditorSettingsName", "Destructable Mesh Editor"),
 			LOCTEXT("ContentEditorsDestructableMeshEditorSettingsDescription", "Change the behavior of the Destructable Mesh Editor."),
-			TWeakObjectPtr<UObject>(GetMutableDefault<UDestructableMeshEditorSettings>())
+			GetMutableDefault<UDestructableMeshEditorSettings>()
 		);*/
 	}
 
diff --git a/Engine/Source/Editor/EditorStyle/Private/SlateEditorStyle.h b/Engine/Source/Editor/EditorStyle/Private/SlateEditorStyle.h
--- a/Engine/Source/Editor/EditorStyle/Private/SlateEditorStyle.h
+++ b/Engine/Source/Editor/EditorStyle/Private/SlateEditorStyle.h
@@ -28,7 +28,7 @@ public:
 			SettingsModule->RegisterSettings( "Editor", "General", "Appearance",
 				NSLOCTEXT("EditorStyle", "Appearance_UserSettingsName", "Appearance"),
 				NSLOCTEXT("EditorStyle", "Appearance_UserSettingsDescription", "Customize the look of the editor."),
-				TWeakObjectPtr<UObject>( Settings.Get() )
+				Settings
 			);
 		}
 #endif
diff --git a/Engine/Source/Editor/NewsFeed/Private/NewsFeedModule.cpp b/Engine/Source/Editor/NewsFeed/Private/NewsFeedModule.cpp
--- a/Engine/Source/Editor/NewsFeed/Private/NewsFeedModule.cpp
+++ b/Engine/Source/Editor/NewsFeed/Private/NewsFeedModule.cpp
@@ -30,7 +30,7 @@ public:
 			SettingsModule->RegisterSettings("Editor", "LevelEditor", "NewsFeed",
 				LOCTEXT("NewsFeedSettingsName", "News Feed"),
 				LOCTEXT("NewsFeedSettingsDescription", "Change how you receive and view news updates."),
-				TWeakObjectPtr<UObject>(UNewsFeedSettings::StaticClass()->GetDefaultObject())
+				GetMutableDefault<UNewsFeedSettings>()
 			);
 		}
 	}
diff --git a/Engine/Source/Editor/ProjectSettingsViewer/Private/ProjectSettingsViewerModule.cpp b/Engine/Source/Editor/ProjectSettingsViewer/Private/ProjectSettingsViewerModule.cpp
--- a/Engine/Source/Editor/ProjectSettingsViewer/Private/ProjectSettingsViewerModule.cpp
+++ b/Engine/Source/Editor/ProjectSettingsViewer/Private/ProjectSettingsViewerModule.cpp
@@ -88,49 +88,49 @@ protected:
 		SettingsModule.RegisterSettings("Project", "Engine", "General",
 			LOCTEXT("GeneralEngineSettingsName", "General Settings"),
 			LOCTEXT("ProjectGeneralSettingsDescription", "General options and defaults for the game engine."),
-			TWeakObjectPtr<UObject>(GetMutableDefault<UEngine>())
+			GetMutableDefault<UEngine>()
 		);
 
 		// Audio settings
 		SettingsModule.RegisterSettings("Project", "Engine", "Audio",
 			LOCTEXT("EngineAudioSettingsName", "Audio"),
 			LOCTEXT("ProjectAudioSettingsDescription", "Audio settings."),
-			TWeakObjectPtr<UObject>(GetMutableDefault<UAudioSettings>())
-			);
+			GetMutableDefault<UAudioSettings>()
+		);
 
 		// packaging settings
 		SettingsModule.RegisterSettings("Project", "Engine", "Collision",
 			LOCTEXT("ProjectCollisionSettingsName", "Collision"),
 			LOCTEXT("ProjectCollisionSettingsDescription", "Set up and modify collision settings."),
-			TWeakObjectPtr<UObject>(GetMutableDefault<UCollisionProfile>())
+			GetMutableDefault<UCollisionProfile>()
 		);
 
 		// game console settings
 		SettingsModule.RegisterSettings("Project", "Engine", "Console",
 			LOCTEXT("ProjectConsoleSettingsName", "Console"),
 			LOCTEXT("ProjectConsoleSettingsDescription", "Configure the in-game input console."),
-			TWeakObjectPtr<UObject>(GetMutableDefault<UConsole>())
+			GetMutableDefault<UConsole>()
 		);
 
 		// Input settings
 		SettingsModule.RegisterSettings("Project", "Engine", "Input",
 			LOCTEXT("EngineInputSettingsName", "Input"),
 			LOCTEXT("ProjectInputSettingsDescription", "Input settings, including default input action and axis bindings."),
-			TWeakObjectPtr<UObject>(GetMutableDefault<UInputSettings>())
+			GetMutableDefault<UInputSettings>()
 		);
 
 		// navigation system
 		SettingsModule.RegisterSettings("Project", "Engine", "NavigationSystem",
 			LOCTEXT("NavigationSystemSettingsName", "Navigation System"),
 			LOCTEXT("NavigationSystemSettingsDescription", "Settings for the navigation system."),
-			TWeakObjectPtr<UObject>(GetMutableDefault<UNavigationSystem>())
+			GetMutableDefault<UNavigationSystem>()
 		);
 
 		// navigation mesh
 		SettingsModule.RegisterSettings("Project", "Engine", "NavigationMesh",
 			LOCTEXT("NavigationMeshSettingsName", "Navigation Mesh"),
 			LOCTEXT("NavigationMeshSettingsDescription", "Settings for the navigation mesh."),
-			TWeakObjectPtr<UObject>(GetMutableDefault<ARecastNavMesh>())
+			GetMutableDefault<ARecastNavMesh>()
 		);
 
 /*
@@ -138,22 +138,22 @@ protected:
 		SettingsModule.RegisterSettings("Project", "Engine", "NetworkManager",
 			LOCTEXT("GameNetworkManagerSettingsName", "Network Manager"),
 			LOCTEXT("GameNetworkManagerSettingsDescription", "Game Network Manager settings."),
-			TWeakObjectPtr<UObject>(GetMutableDefault<UGameNetworkManagerSettings>())
+			GetMutableDefault<UGameNetworkManagerSettings>()
 		);*/
 
 		// Physics settings
 		SettingsModule.RegisterSettings("Project", "Engine", "Physics",
 			LOCTEXT("EnginePhysicsSettingsName", "Physics"),
 			LOCTEXT("ProjectPhysicsSettingsDescription", "Default physics settings."),
-			TWeakObjectPtr<UObject>(GetMutableDefault<UPhysicsSettings>())
-			);	
+			GetMutableDefault<UPhysicsSettings>()
+		);
 
 		// Rendering settings
 		SettingsModule.RegisterSettings("Project", "Engine", "Rendering",
 			LOCTEXT("EngineRenderingSettingsName", "Rendering"),
 			LOCTEXT("ProjectRenderingSettingsDescription", "Rendering settings."),
-			TWeakObjectPtr<UObject>(GetMutableDefault<URendererSettings>())
-			);	
+			GetMutableDefault<URendererSettings>()
+		);
 	}
 
 	/**
@@ -165,46 +165,45 @@ protected:
 		SettingsModule.RegisterSettings("Project", "Game", "General",
 			LOCTEXT("GeneralGameSettingsName", "Description"),
 			LOCTEXT("GeneralGameSettingsDescription", "Descriptions and other information about your game."),
-			TWeakObjectPtr<UObject>(GetMutableDefault<UGeneralProjectSettings>())
+			GetMutableDefault<UGeneralProjectSettings>()
 		);
 
 		// map related settings
 		SettingsModule.RegisterSettings("Project", "Game", "Maps",
 			LOCTEXT("GameMapsSettingsName", "Maps & Modes"),
 			LOCTEXT("GameMapsSettingsDescription", "Default maps, game modes and other map related settings."),
-			TWeakObjectPtr<UObject>(GetMutableDefault<UGameMapsSettings>())
+			GetMutableDefault<UGameMapsSettings>()
 		);
 
 		// packaging settings
 		SettingsModule.RegisterSettings("Project", "Game", "Packaging",
 			LOCTEXT("ProjectPackagingSettingsName", "Packaging"),
 			LOCTEXT("ProjectPackagingSettingsDescription", "Fine tune how your project is packaged for release."),
-			TWeakObjectPtr<UObject>(GetMutableDefault<UProjectPackagingSettings>())
+			GetMutableDefault<UProjectPackagingSettings>()
 		);
 
 		// movie settings
 		SettingsModule.RegisterSettings("Project", "Game", "Movies",
 			LOCTEXT("MovieSettingsName", "Movies"),
 			LOCTEXT("MovieSettingsDescription", "Movie player settings"),
-			TWeakObjectPtr<UObject>(GetMutableDefault<UMoviePlayerSettings>())
+			GetMutableDefault<UMoviePlayerSettings>()
 		);
 
 		// update the visible packaging enum values
 		UProjectPackagingSettings::UpdateBuildConfigurationVisibility();
-
 /*
 		// game session
 		SettingsModule.RegisterSettings("Project", "Game", "GameSession",
 			LOCTEXT("GameSessionettingsName", "Game Session"),
 			LOCTEXT("GameSessionSettingsDescription", "Game Session settings."),
-			TWeakObjectPtr<UObject>(GetMutableDefault<UGameSessionSettings>())
+			GetMutableDefault<UGameSessionSettings>()
 		);
 
 		// head-up display
 		SettingsModule.RegisterSettings("Project", "Game", "HUD",
 			LOCTEXT("HudSettingsName", "HUD"),
 			LOCTEXT("HudSettingsDescription", "Head-up display (HUD) settings."),
-			TWeakObjectPtr<UObject>(GetMutableDefault<UHudSettings>())
+			GetMutableDefault<UHudSettings>()
 		);*/
 	}
 
diff --git a/Engine/Source/Editor/TextureEditor/Private/TextureEditorModule.cpp b/Engine/Source/Editor/TextureEditor/Private/TextureEditorModule.cpp
--- a/Engine/Source/Editor/TextureEditor/Private/TextureEditorModule.cpp
+++ b/Engine/Source/Editor/TextureEditor/Private/TextureEditorModule.cpp
@@ -59,7 +59,7 @@ public:
 			SettingsModule->RegisterSettings("Editor", "ContentEditors", "TextureEditor",
 				LOCTEXT("TextureEditorSettingsName", "Texture Editor"),
 				LOCTEXT("TextureEditorSettingsDescription", "Configure the look and feel of the Texture Editor."),
-				TWeakObjectPtr<UObject>(GetMutableDefault<UTextureEditorSettings>())
+				GetMutableDefault<UTextureEditorSettings>()
 			);
 		}
 	}
