Fix include path for new CookerSettings.h location

diff --git a/Engine/Source/Editor/ProjectSettingsViewer/Private/ProjectSettingsViewerModule.cpp b/Engine/Source/Editor/ProjectSettingsViewer/Private/ProjectSettingsViewerModule.cpp
--- a/Engine/Source/Editor/ProjectSettingsViewer/Private/ProjectSettingsViewerModule.cpp
+++ b/Engine/Source/Editor/ProjectSettingsViewer/Private/ProjectSettingsViewerModule.cpp
@@ -3,7 +3,7 @@
 #include "ProjectSettingsViewerPrivatePCH.h"
 #include "Engine/Console.h"
 #include "ProjectTargetPlatformEditor.h"
-#include "Settings/CookerSettings.h"
+#include "CookerSettings.h"
 
 
 #define LOCTEXT_NAMESPACE "FProjectSettingsViewerModule"
