Corrected bad copy-pasted comment in project settings viewer code.

diff --git a/Engine/Source/Editor/ProjectSettingsViewer/Private/ProjectSettingsViewerModule.cpp b/Engine/Source/Editor/ProjectSettingsViewer/Private/ProjectSettingsViewerModule.cpp
--- a/Engine/Source/Editor/ProjectSettingsViewer/Private/ProjectSettingsViewerModule.cpp
+++ b/Engine/Source/Editor/ProjectSettingsViewer/Private/ProjectSettingsViewerModule.cpp
@@ -91,7 +91,7 @@ protected:
 			GetMutableDefault<UAudioSettings>()
 		);
 
-		// packaging settings
+		// collision settings
 		SettingsModule.RegisterSettings("Project", "Engine", "Collision",
 			LOCTEXT("ProjectCollisionSettingsName", "Collision"),
 			LOCTEXT("ProjectCollisionSettingsDescription", "Set up and modify collision settings."),
