Renamed the "Game" section under "Project Settings" to be "Project"

Suggested via an email discussion, as not everyone using UE4 is making a game.

diff --git a/Engine/Source/Developer/Settings/Private/SettingsModule.cpp b/Engine/Source/Developer/Settings/Private/SettingsModule.cpp
--- a/Engine/Source/Developer/Settings/Private/SettingsModule.cpp
+++ b/Engine/Source/Developer/Settings/Private/SettingsModule.cpp
@@ -83,7 +83,7 @@ public:
 		// @todo gmp: move this into the modules that own these setting categories
 		TSharedRef<FSettingsContainer> ProjectSettingsContainer = FindOrAddContainer("Project");
 		ProjectSettingsContainer->Describe(LOCTEXT("ProjectSettingsSubMenuLabel", "Project Settings"), LOCTEXT("ProjectSettingsSubMenuToolTip", "Change the settings of the currently loaded project"), NAME_None);
-		ProjectSettingsContainer->DescribeCategory("Game", LOCTEXT("ProjectGameCategoryName", "Game"), LOCTEXT("ProjectGameCategoryDescription", "Project settings description text here"), "SettingsEditor.Category_Game");
+		ProjectSettingsContainer->DescribeCategory("Project", LOCTEXT("ProjectProjectCategoryName", "Project"), LOCTEXT("ProjectProjectCategoryDescription", "Project settings description text here"), "SettingsEditor.Category_Project");
 		ProjectSettingsContainer->DescribeCategory("Engine", LOCTEXT("ProjectEngineCategoryName", "Engine"), LOCTEXT("ProjectEngineCategoryDescription", "Project settings description text here"), "SettingsEditor.Category_Engine");
 		ProjectSettingsContainer->DescribeCategory("Platforms", LOCTEXT("ProjectPlatformsCategoryName", "Platforms"), LOCTEXT("ProjectPlatformsCategoryDescription", "Platform settings description text here"), "SettingsEditor.Category_Platforms");
 		ProjectSettingsContainer->DescribeCategory("Plugins", LOCTEXT("ProjectPluginsCategoryName", "Plugins"), LOCTEXT("ProjectPluginsCategoryDescription", "Plugins settings description text here"), "SettingsEditor.Category_Plugins");
diff --git a/Engine/Source/Editor/LevelEditor/Private/LevelEditorToolBar.cpp b/Engine/Source/Editor/LevelEditor/Private/LevelEditorToolBar.cpp
--- a/Engine/Source/Editor/LevelEditor/Private/LevelEditorToolBar.cpp
+++ b/Engine/Source/Editor/LevelEditor/Private/LevelEditorToolBar.cpp
@@ -643,7 +643,7 @@ TSharedRef< SWidget > FLevelEditorToolBar::GenerateOpenGameSettingsMenu( TShared
 		LOCTEXT("ProjectSettingsMenuLabel", "Project Settings..."),
 		LOCTEXT("ProjectSettingsMenuToolTip", "Change the settings of the currently loaded project"),
 		FSlateIcon(),
-		FUIAction(FExecuteAction::CreateStatic(&Local::OpenSettings, FName("Project"), FName("Game"), FName("General")))
+		FUIAction(FExecuteAction::CreateStatic(&Local::OpenSettings, FName("Project"), FName("Project"), FName("General")))
 	);
 
 #undef LOCTEXT_NAMESPACE
diff --git a/Engine/Source/Editor/MainFrame/Private/Frame/MainFrameActions.cpp b/Engine/Source/Editor/MainFrame/Private/Frame/MainFrameActions.cpp
--- a/Engine/Source/Editor/MainFrame/Private/Frame/MainFrameActions.cpp
+++ b/Engine/Source/Editor/MainFrame/Private/Frame/MainFrameActions.cpp
@@ -716,7 +716,7 @@ void FMainFrameActionCallbacks::OpenIDE()
 
 void FMainFrameActionCallbacks::PackagingSettings()
 {
-	FModuleManager::LoadModuleChecked<ISettingsModule>("Settings").ShowViewer("Project", "Game", "Packaging");
+	FModuleManager::LoadModuleChecked<ISettingsModule>("Settings").ShowViewer("Project", "Project", "Packaging");
 }
 
 void FMainFrameActionCallbacks::SwitchProjectByIndex( int32 ProjectIndex )
diff --git a/Engine/Source/Editor/MainFrame/Private/Menus/MainMenu.cpp b/Engine/Source/Editor/MainFrame/Private/Menus/MainMenu.cpp
--- a/Engine/Source/Editor/MainFrame/Private/Menus/MainMenu.cpp
+++ b/Engine/Source/Editor/MainFrame/Private/Menus/MainMenu.cpp
@@ -132,7 +132,7 @@ void FMainMenu::FillEditMenu( FMenuBuilder& MenuBuilder, const TSharedRef< FExte
 				LOCTEXT("ProjectSettingsMenuLabel", "Project Settings..."),
 				LOCTEXT("ProjectSettingsMenuToolTip", "Change the settings of the currently loaded project"),
 				FSlateIcon(),
-				FUIAction(FExecuteAction::CreateStatic(&FSettingsMenu::OpenSettings, FName("Project"), FName("Game"), FName("General")))
+				FUIAction(FExecuteAction::CreateStatic(&FSettingsMenu::OpenSettings, FName("Project"), FName("Project"), FName("General")))
 			);
 		}
 	}
diff --git a/Engine/Source/Editor/ProjectSettingsViewer/Private/ProjectSettingsViewerModule.cpp b/Engine/Source/Editor/ProjectSettingsViewer/Private/ProjectSettingsViewerModule.cpp
--- a/Engine/Source/Editor/ProjectSettingsViewer/Private/ProjectSettingsViewerModule.cpp
+++ b/Engine/Source/Editor/ProjectSettingsViewer/Private/ProjectSettingsViewerModule.cpp
@@ -49,7 +49,7 @@ public:
 		if (SettingsModule != nullptr)
 		{
 			RegisterEngineSettings(*SettingsModule);
-			RegisterGameSettings(*SettingsModule);
+			RegisterProjectSettings(*SettingsModule);
 
 			SettingsModule->RegisterViewer("Project", *this);
 		}
@@ -163,26 +163,26 @@ protected:
 	}
 
 	/**
-	 * Registers Game settings.
+	 * Registers Project settings.
 	 */
-	void RegisterGameSettings( ISettingsModule& SettingsModule )
+	void RegisterProjectSettings( ISettingsModule& SettingsModule )
 	{
 		// general project settings
-		SettingsModule.RegisterSettings("Project", "Game", "General",
+		SettingsModule.RegisterSettings("Project", "Project", "General",
 			LOCTEXT("GeneralGameSettingsName", "Description"),
-			LOCTEXT("GeneralGameSettingsDescription", "Descriptions and other information about your game."),
+			LOCTEXT("GeneralGameSettingsDescription", "Descriptions and other information about your project."),
 			GetMutableDefault<UGeneralProjectSettings>()
 		);
 
 		// map related settings
-		SettingsModule.RegisterSettings("Project", "Game", "Maps",
+		SettingsModule.RegisterSettings("Project", "Project", "Maps",
 			LOCTEXT("GameMapsSettingsName", "Maps & Modes"),
 			LOCTEXT("GameMapsSettingsDescription", "Default maps, game modes and other map related settings."),
 			GetMutableDefault<UGameMapsSettings>()
 		);
 
 		// packaging settings
-		SettingsModule.RegisterSettings("Project", "Game", "Packaging",
+		SettingsModule.RegisterSettings("Project", "Project", "Packaging",
 			LOCTEXT("ProjectPackagingSettingsName", "Packaging"),
 			LOCTEXT("ProjectPackagingSettingsDescription", "Fine tune how your project is packaged for release."),
 			GetMutableDefault<UProjectPackagingSettings>()
@@ -190,28 +190,28 @@ protected:
 
 		// platforms settings
 		TWeakPtr<SWidget> ProjectTargetPlatformEditorPanel = FModuleManager::LoadModuleChecked<IProjectTargetPlatformEditorModule>("ProjectTargetPlatformEditor").CreateProjectTargetPlatformEditorPanel();
-		SettingsModule.RegisterSettings("Project", "Game", "SupportedPlatforms",
+		SettingsModule.RegisterSettings("Project", "Project", "SupportedPlatforms",
 			LOCTEXT("ProjectSupportedPlatformsSettingsName", "Supported Platforms"),
 			LOCTEXT("ProjectSupportedPlatformsSettingsDescription", "Specify which platforms your project supports."),
 			ProjectTargetPlatformEditorPanel.Pin().ToSharedRef()
 		);
 
 		// movie settings
-		SettingsModule.RegisterSettings("Project", "Game", "Movies",
+		SettingsModule.RegisterSettings("Project", "Project", "Movies",
 			LOCTEXT("MovieSettingsName", "Movies"),
 			LOCTEXT("MovieSettingsDescription", "Movie player settings"),
 			GetMutableDefault<UMoviePlayerSettings>()
 		);
 /*
 		// game session
-		SettingsModule.RegisterSettings("Project", "Game", "GameSession",
+		SettingsModule.RegisterSettings("Project", "Project", "GameSession",
 			LOCTEXT("GameSessionettingsName", "Game Session"),
 			LOCTEXT("GameSessionSettingsDescription", "Game Session settings."),
 			GetMutableDefault<UGameSessionSettings>()
 		);
 
 		// head-up display
-		SettingsModule.RegisterSettings("Project", "Game", "HUD",
+		SettingsModule.RegisterSettings("Project", "Project", "HUD",
 			LOCTEXT("HudSettingsName", "HUD"),
 			LOCTEXT("HudSettingsDescription", "Head-up display (HUD) settings."),
 			GetMutableDefault<UHudSettings>()
@@ -234,19 +234,19 @@ protected:
 			SettingsModule->UnregisterSettings("Project", "Engine", "NavigationSystem");
 			SettingsModule->UnregisterSettings("Project", "Engine", "NavigationMesh");
 			SettingsModule->UnregisterSettings("Project", "Engine", "Input");
-			SettingsModule->UnregisterSettings("Project", "Game", "Collision");
+			SettingsModule->UnregisterSettings("Project", "Engine", "Collision");
 			SettingsModule->UnregisterSettings("Project", "Engine", "Physics");
 			SettingsModule->UnregisterSettings("Project", "Engine", "Rendering");
 //			SettingsModule->UnregisterSettings("Project", "Engine", "NetworkManager");
 
-			// game settings
-			SettingsModule->UnregisterSettings("Project", "Game", "General");
-			SettingsModule->UnregisterSettings("Project", "Game", "Maps");
-			SettingsModule->UnregisterSettings("Project", "Game", "Packaging");
-			SettingsModule->UnregisterSettings("Project", "Game", "SupportedPlatforms");
-			SettingsModule->UnregisterSettings("Project", "Game", "Movies");
-//			SettingsModule->UnregisterSettings("Project", "Game", "GameSession");
-//			SettingsModule->UnregisterSettings("Project", "Game", "HUD");
+			// project settings
+			SettingsModule->UnregisterSettings("Project", "Project", "General");
+			SettingsModule->UnregisterSettings("Project", "Project", "Maps");
+			SettingsModule->UnregisterSettings("Project", "Project", "Packaging");
+			SettingsModule->UnregisterSettings("Project", "Project", "SupportedPlatforms");
+			SettingsModule->UnregisterSettings("Project", "Project", "Movies");
+//			SettingsModule->UnregisterSettings("Project", "Project", "GameSession");
+//			SettingsModule->UnregisterSettings("Project", "Project", "HUD");
 		}
 	}
 
diff --git a/Engine/Source/Editor/ProjectTargetPlatformEditor/Private/ProjectTargetPlatformEditorModule.cpp b/Engine/Source/Editor/ProjectTargetPlatformEditor/Private/ProjectTargetPlatformEditorModule.cpp
--- a/Engine/Source/Editor/ProjectTargetPlatformEditor/Private/ProjectTargetPlatformEditorModule.cpp
+++ b/Engine/Source/Editor/ProjectTargetPlatformEditor/Private/ProjectTargetPlatformEditorModule.cpp
@@ -41,7 +41,7 @@ public:
 			LOCTEXT("SupportedPlatformsMenuLabel", "Supported Platforms..."),
 			LOCTEXT("SupportedPlatformsMenuToolTip", "Change which platforms this project supports"),
 			FSlateIcon(),
-			FUIAction(FExecuteAction::CreateStatic(&Local::OpenSettings, FName("Project"), FName("Game"), FName("SupportedPlatforms")))
+			FUIAction(FExecuteAction::CreateStatic(&Local::OpenSettings, FName("Project"), FName("Project"), FName("SupportedPlatforms")))
 			);
 	}
 
