Editor - UDeveloperSettings is now available to create settings objects that can be auto-discovered by the editor without building a plugin or some other system to register the settings with the settings viewer.

diff --git a/Engine/Source/Developer/Settings/Private/SettingsModule.cpp b/Engine/Source/Developer/Settings/Private/SettingsModule.cpp
--- a/Engine/Source/Developer/Settings/Private/SettingsModule.cpp
+++ b/Engine/Source/Developer/Settings/Private/SettingsModule.cpp
@@ -85,6 +85,7 @@ public:
 		TSharedRef<FSettingsContainer> ProjectSettingsContainer = FindOrAddContainer("Project");
 		ProjectSettingsContainer->Describe(LOCTEXT("ProjectSettingsSubMenuLabel", "Project Settings"), LOCTEXT("ProjectSettingsSubMenuToolTip", "Change the settings of the currently loaded project"), NAME_None);
 		ProjectSettingsContainer->DescribeCategory("Project", LOCTEXT("ProjectProjectCategoryName", "Project"), LOCTEXT("ProjectProjectCategoryDescription", "Project settings description text here"));
+		ProjectSettingsContainer->DescribeCategory("Game", LOCTEXT("ProjectGameCategoryName", "Game"), LOCTEXT("ProjectGameCategoryDescription", "Game settings description text here"));
 		ProjectSettingsContainer->DescribeCategory("Engine", LOCTEXT("ProjectEngineCategoryName", "Engine"), LOCTEXT("ProjectEngineCategoryDescription", "Project settings description text here"));
 		ProjectSettingsContainer->DescribeCategory("Editor", LOCTEXT("ProjectEditorCategoryName", "Editor"), LOCTEXT("ProjectEditorCategoryDescription", "Project settings specific to the editor"));
 		ProjectSettingsContainer->DescribeCategory("Platforms", LOCTEXT("ProjectPlatformsCategoryName", "Platforms"), LOCTEXT("ProjectPlatformsCategoryDescription", "Platform settings description text here"));
diff --git a/Engine/Source/Editor/ProjectSettingsViewer/Private/ProjectSettingsViewerModule.cpp b/Engine/Source/Editor/ProjectSettingsViewer/Private/ProjectSettingsViewerModule.cpp
--- a/Engine/Source/Editor/ProjectSettingsViewer/Private/ProjectSettingsViewerModule.cpp
+++ b/Engine/Source/Editor/ProjectSettingsViewer/Private/ProjectSettingsViewerModule.cpp
@@ -13,6 +13,7 @@
 #include "SDockTab.h"
 
 #include "CookerSettings.h"
+#include "Runtime/Engine/Classes/Engine/DeveloperSettings.h"
 #include "Runtime/Engine/Classes/Engine/RendererSettings.h"
 #include "Runtime/Engine/Classes/Engine/UserInterfaceSettings.h"
 #include "Runtime/Engine/Classes/Sound/AudioSettings.h"
@@ -31,6 +32,14 @@
 
 static const FName ProjectSettingsTabName("ProjectSettings");
 
+/** Holds auto discovered settings information so that they can be unloaded automatically when refreshing. */
+struct FRegisteredSettings
+{
+	FName ContainerName;
+	FName CategoryName;
+	FName SectionName;
+};
+
 
 /**
  * Implements the ProjectSettingsViewer module.
@@ -38,9 +47,16 @@ static const FName ProjectSettingsTabName("ProjectSettings");
 class FProjectSettingsViewerModule
 	: public IModuleInterface
 	, public ISettingsViewer
+	, public FTickableEditorObject
 {
 public:
 
+	FProjectSettingsViewerModule()
+	{
+		bTicking = false;
+		bInvalidated = false;
+	}
+
 	// ISettingsViewer interface
 
 	virtual void ShowSettings( const FName& CategoryName, const FName& SectionName ) override
@@ -73,18 +89,36 @@ public:
 			RegisterProjectSettings(*SettingsModule);
 			RegisterEditorSettings(*SettingsModule);
 
+			RegisterAutoDiscoveredSettings(*SettingsModule);
+
 			SettingsModule->RegisterViewer("Project", *this);
 		}
 
 		FGlobalTabmanager::Get()->RegisterNomadTabSpawner(ProjectSettingsTabName, FOnSpawnTab::CreateRaw(this, &FProjectSettingsViewerModule::HandleSpawnSettingsTab))
 			.SetDisplayName(LOCTEXT("ProjectSettingsTabTitle", "Project Settings"))
 			.SetMenuType(ETabSpawnerMenuType::Hidden);
+
+		FModuleManager::Get().OnModulesChanged().AddRaw(this, &FProjectSettingsViewerModule::ModulesChangesCallback);
+
+		bTicking = true;
 	}
 
 	virtual void ShutdownModule() override
 	{
+		bTicking = false;
+
 		FGlobalTabmanager::Get()->UnregisterNomadTabSpawner(ProjectSettingsTabName);
 		UnregisterSettings();
+
+		FModuleManager::Get().OnModulesChanged().RemoveAll(this);
+	}
+
+	void ModulesChangesCallback(FName ModuleName, EModuleChangeReason ReasonForChange)
+	{
+		if ( ReasonForChange == EModuleChangeReason::ModuleLoaded )
+		{
+			bInvalidated = true;
+		}
 	}
 
 	virtual bool SupportsDynamicReloading() override
@@ -174,14 +208,6 @@ protected:
 			GetMutableDefault<UAnimationSettings>()
 		);
 
-		/*
-		// network settings
-		SettingsModule.RegisterSettings("Project", "Engine", "NetworkManager",
-			LOCTEXT("GameNetworkManagerSettingsName", "Network Manager"),
-			LOCTEXT("GameNetworkManagerSettingsDescription", "Game Network Manager settings."),
-			GetMutableDefault<UGameNetworkManagerSettings>()
-		);*/
-
 		// network settings
 		SettingsModule.RegisterSettings("Project", "Engine", "Network",
 			LOCTEXT("NetworkSettingsName", "Network"),
@@ -260,22 +286,48 @@ protected:
 			LOCTEXT("MovieSettingsDescription", "Movie player settings"),
 			GetMutableDefault<UMoviePlayerSettings>()
 		);
-/*
-		// game session
-		SettingsModule.RegisterSettings("Project", "Project", "GameSession",
-			LOCTEXT("GameSessionettingsName", "Game Session"),
-			LOCTEXT("GameSessionSettingsDescription", "Game Session settings."),
-			GetMutableDefault<UGameSessionSettings>()
-		);
-
-		// head-up display
-		SettingsModule.RegisterSettings("Project", "Project", "HUD",
-			LOCTEXT("HudSettingsName", "HUD"),
-			LOCTEXT("HudSettingsDescription", "Head-up display (HUD) settings."),
-			GetMutableDefault<UHudSettings>()
-		);*/
 	}
 
+	void RegisterAutoDiscoveredSettings(ISettingsModule& SettingsModule)
+	{
+		// Find game object
+		for ( TObjectIterator<UDeveloperSettings> SettingsIt(RF_NoFlags); SettingsIt; ++SettingsIt )
+		{
+			if ( UDeveloperSettings* Settings = *SettingsIt )
+			{
+				// Only Add the CDO of any UDeveloperSettings objects.
+				if ( Settings->HasAnyFlags(RF_ClassDefaultObject) && !Settings->GetClass()->HasAnyCastFlag(CLASS_Deprecated | CLASS_Abstract) )
+				{
+					FRegisteredSettings Registered;
+					Registered.ContainerName = Settings->GetContainerName();
+					Registered.CategoryName = Settings->GetCategoryName();
+					Registered.SectionName = Settings->GetSectionName();
+
+					TSharedPtr<SWidget> CustomWidget = Settings->GetCustomSettingsWidget();
+					if ( CustomWidget.IsValid() )
+					{
+						// Add Settings
+						SettingsModule.RegisterSettings(Registered.ContainerName, Registered.CategoryName, Registered.SectionName,
+							Settings->GetSectionText(),
+							Settings->GetSectionDescription(),
+							CustomWidget.ToSharedRef()
+							);
+					}
+					else
+					{
+						// Add Settings
+						SettingsModule.RegisterSettings(Registered.ContainerName, Registered.CategoryName, Registered.SectionName,
+							Settings->GetSectionText(),
+							Settings->GetSectionDescription(),
+							Settings
+							);
+					}
+
+					AutoDiscoveredSettings.Add(Registered);
+				}
+			}
+		}
+	}
 
 	/**
 	 * Registers Editor settings.
@@ -307,6 +359,7 @@ protected:
 		if (SettingsModule != nullptr)
 		{
 			SettingsModule->UnregisterViewer("Project");
+
 			// engine settings
 			SettingsModule->UnregisterSettings("Project", "Engine", "General");
 			SettingsModule->UnregisterSettings("Project", "Engine", "CrowdManager");
@@ -316,7 +369,6 @@ protected:
 			SettingsModule->UnregisterSettings("Project", "Engine", "Collision");
 			SettingsModule->UnregisterSettings("Project", "Engine", "Physics");
 			SettingsModule->UnregisterSettings("Project", "Engine", "Rendering");
-//			SettingsModule->UnregisterSettings("Project", "Engine", "NetworkManager");
 
 			// project settings
 			SettingsModule->UnregisterSettings("Project", "Project", "General");
@@ -328,11 +380,52 @@ protected:
 			// Editor settings
 			SettingsModule->UnregisterSettings("Editor", "Editor", "Appearance");
 
-//			SettingsModule->UnregisterSettings("Project", "Project", "GameSession");
-//			SettingsModule->UnregisterSettings("Project", "Project", "HUD");
+			UnregisterAutoDiscoveredSettings(*SettingsModule);
 		}
 	}
 
+	void UnregisterAutoDiscoveredSettings(ISettingsModule& SettingsModule)
+	{
+		// Unregister any auto discovers settings.
+		for ( const FRegisteredSettings& Settings : AutoDiscoveredSettings )
+		{
+			SettingsModule.UnregisterSettings(Settings.ContainerName, Settings.CategoryName, Settings.SectionName);
+		}
+
+		AutoDiscoveredSettings.Reset();
+	}
+
+private:
+
+	/** FTickableEditorObject interface */
+	void Tick(float DeltaTime) override
+	{
+		if ( bInvalidated )
+		{
+			ISettingsModule* SettingsModule = FModuleManager::GetModulePtr<ISettingsModule>("Settings");
+
+			if ( SettingsModule != nullptr )
+			{
+				UnregisterAutoDiscoveredSettings(*SettingsModule);
+				RegisterAutoDiscoveredSettings(*SettingsModule);
+			}
+
+			bInvalidated = false;
+		}
+	}
+
+	bool IsTickable() const override
+	{
+		return bTicking;
+	}
+
+	TStatId GetStatId() const override
+	{
+		RETURN_QUICK_DECLARE_CYCLE_STAT(FProjectSettingsViewerModule, STATGROUP_Tickables);
+	}
+
+	/** FTickableEditorObject interface */
+
 private:
 
 	/** Handles creating the project settings tab. */
@@ -366,6 +459,15 @@ private:
 
 	/** Holds a pointer to the settings editor's view model. */
 	TWeakPtr<ISettingsEditorModel> SettingsEditorModelPtr;
+
+	/** The list of auto discovered settings that need to be unregistered. */
+	TArray<FRegisteredSettings> AutoDiscoveredSettings;
+
+	/** When new modules are loaded we invalidate the list of autodiscovered settings and re-discover them. */
+	bool bInvalidated;
+
+	/** Should we be ticking? */
+	bool bTicking;
 };
 
 
diff --git a/Engine/Source/Runtime/Engine/Classes/Engine/DeveloperSettings.h b/Engine/Source/Runtime/Engine/Classes/Engine/DeveloperSettings.h
new file mode 100644
--- /dev/null
+++ b/Engine/Source/Runtime/Engine/Classes/Engine/DeveloperSettings.h
@@ -0,0 +1,51 @@
+// Copyright 1998-2015 Epic Games, Inc. All Rights Reserved.
+
+#pragma once
+
+#include "DeveloperSettings.generated.h"
+
+/** The different types of settings container that can be used. */
+UENUM()
+enum class EDeveloperSettingsContainer : uint8
+{
+	Project,
+	Editor
+};
+
+/**
+ * The base class of any auto discovered settings object.
+ */
+UCLASS(abstract)
+class ENGINE_API UDeveloperSettings : public UObject
+{
+	GENERATED_BODY()
+
+public:
+	UDeveloperSettings(const FObjectInitializer& ObjectInitializer);
+
+	/** Gets the settings container name for the settings, either Project or Editor */
+	FName GetContainerName() const;
+	/** Gets the category for the settings, some high level grouping like, Editor, Engine, Game...etc. */
+	FName GetCategoryName() const;
+	/** The unique name for your section of settings, uses the class's FName. */
+	FName GetSectionName() const;
+
+	/** Gets the section text, uses the classes DisplayName by default. */
+	virtual FText GetSectionText() const;
+	/** Gets the description for the section, uses the classes ToolTip by default. */
+	virtual FText GetSectionDescription() const;
+
+	/** Gets a custom widget for the settings.  This is only for very custom situations. */
+	virtual TSharedPtr<SWidget> GetCustomSettingsWidget() const;
+
+private:
+	/**
+	 * The category name to use, overrides the one detected by looking at the config=... class metadata.
+	 * Arbitrary category names are not supported, this must map to an existing category we support in the settings
+	 * viewer.
+	 */
+	FName CategoryName;
+
+	/** The container for these settings. */
+	EDeveloperSettingsContainer Container;
+};
diff --git a/Engine/Source/Runtime/Engine/Private/DeveloperSettings.cpp b/Engine/Source/Runtime/Engine/Private/DeveloperSettings.cpp
new file mode 100644
--- /dev/null
+++ b/Engine/Source/Runtime/Engine/Private/DeveloperSettings.cpp
@@ -0,0 +1,72 @@
+// Copyright 1998-2015 Epic Games, Inc. All Rights Reserved.
+
+#include "EnginePrivate.h"
+#include "Engine/DeveloperSettings.h"
+
+UDeveloperSettings::UDeveloperSettings(const FObjectInitializer& ObjectInitializer)
+	: UObject(ObjectInitializer)
+{
+	CategoryName = NAME_None;
+	Container = EDeveloperSettingsContainer::Project;
+}
+
+FName UDeveloperSettings::GetContainerName() const
+{
+	FName ProjectName(TEXT("Project"));
+	FName EditorName(TEXT("Editor"));
+
+	switch ( Container )
+	{
+	case EDeveloperSettingsContainer::Project:
+		return ProjectName;
+	case EDeveloperSettingsContainer::Editor:
+		return EditorName;
+	default:
+		check(false);
+		return NAME_None;
+	}
+}
+
+FName UDeveloperSettings::GetCategoryName() const
+{
+	if ( CategoryName != NAME_None )
+	{
+		return CategoryName;
+	}
+
+	FName ConfigName = GetClass()->ClassConfigName;
+	if ( ConfigName == NAME_Engine || ConfigName == NAME_Input )
+	{
+		return NAME_Engine;
+	}
+	else if ( ConfigName == NAME_Editor || ConfigName == NAME_EditorSettings || ConfigName == NAME_EditorLayout || ConfigName == NAME_EditorKeyBindings )
+	{
+		return NAME_Editor;
+	}
+	else if ( ConfigName == NAME_Game )
+	{
+		return NAME_Game;
+	}
+
+	return NAME_Engine;
+}
+
+FName UDeveloperSettings::GetSectionName() const
+{
+	return GetClass()->GetFName();
+}
+
+FText UDeveloperSettings::GetSectionText() const
+{
+	return GetClass()->GetDisplayNameText();
+}
+
+FText UDeveloperSettings::GetSectionDescription() const
+{
+	return GetClass()->GetToolTipText();
+}
+
+TSharedPtr<SWidget> UDeveloperSettings::GetCustomSettingsWidget() const
+{
+	return TSharedPtr<SWidget>();
+}
