Editor - Continuing to tweak the settings auto registration system, moving the logic into the SettingsEditor and letting it manage it, since it knows when widgets exist and the information actually needs to be refreshed.  Making several old settings use the new UDeveloperSettings class.

diff --git a/Engine/Source/Developer/SettingsEditor/Private/SettingsEditorModule.cpp b/Engine/Source/Developer/SettingsEditor/Private/SettingsEditorModule.cpp
--- a/Engine/Source/Developer/SettingsEditor/Private/SettingsEditorModule.cpp
+++ b/Engine/Source/Developer/SettingsEditor/Private/SettingsEditorModule.cpp
@@ -3,10 +3,18 @@
 #include "SettingsEditorPrivatePCH.h"
 #include "ISettingsEditorModule.h"
 #include "NotificationManager.h"
-
+#include "ISettingsModule.h"
 
 #define LOCTEXT_NAMESPACE "SSettingsEditor"
 
+/** Holds auto discovered settings information so that they can be unloaded automatically when refreshing. */
+struct FRegisteredSettings
+{
+	FName ContainerName;
+	FName CategoryName;
+	FName SectionName;
+};
+
 
 /** Manages the notification for when the application needs to be restarted due to a settings change */
 class FApplicationRestartRequiredNotification
@@ -104,12 +112,43 @@ class FSettingsEditorModule
 {
 public:
 
+	FSettingsEditorModule()
+		: bAreSettingsStale(true)
+	{
+	}
+
+	// IModuleInterface interface
+
+	virtual void StartupModule() override
+	{
+		FModuleManager::Get().OnModulesChanged().AddRaw(this, &FSettingsEditorModule::ModulesChangesCallback);
+	}
+
+	virtual void ShutdownModule() override
+	{
+		ISettingsModule* SettingsModule = FModuleManager::GetModulePtr<ISettingsModule>("Settings");
+
+		if ( SettingsModule != nullptr )
+		{
+			UnregisterAutoDiscoveredSettings(*SettingsModule);
+		}
+
+		FModuleManager::Get().OnModulesChanged().RemoveAll(this);
+	}
+
 	// ISettingsEditorModule interface
 
 	virtual TSharedRef<SWidget> CreateEditor( const TSharedRef<ISettingsEditorModel>& Model ) override
 	{
-		return SNew(SSettingsEditor, Model)
+		UpdateSettings(true);
+
+		TSharedRef<SWidget> Editor = SNew(SSettingsEditor, Model)
 			.OnApplicationRestartRequired(FSimpleDelegate::CreateRaw(this, &FSettingsEditorModule::OnApplicationRestartRequired));
+
+		ClearStaleEditorWidgets();
+		EditorWidgets.Add(Editor);
+
+		return Editor;
 	}
 
 	virtual ISettingsEditorModelRef CreateModel( const TSharedRef<ISettingsContainer>& SettingsContainer ) override
@@ -127,9 +166,120 @@ public:
 		ApplicationRestartRequiredNotification.SetOnRestartApplicationCallback(InRestartApplicationDelegate);
 	}
 
+private:
+
+	void ModulesChangesCallback(FName ModuleName, EModuleChangeReason ReasonForChange)
+	{
+		ClearStaleEditorWidgets();
+		bAreSettingsStale = true;
+		UpdateSettings();
+	}
+
+	void UpdateSettings(bool bForce = false)
+	{
+		if ( (AnyActiveSettingsEditor() && bAreSettingsStale) || bForce )
+		{
+			bAreSettingsStale = false;
+
+			ISettingsModule* SettingsModule = FModuleManager::GetModulePtr<ISettingsModule>("Settings");
+
+			if ( SettingsModule != nullptr )
+			{
+				UnregisterAutoDiscoveredSettings(*SettingsModule);
+				RegisterAutoDiscoveredSettings(*SettingsModule);
+			}
+		}
+	}
+
+	void ClearStaleEditorWidgets()
+	{
+		for ( int32 i = 0; i < EditorWidgets.Num(); i++ )
+		{
+			if ( !EditorWidgets[i].IsValid() )
+			{
+				EditorWidgets.RemoveAtSwap(i);
+				i--;
+			}
+		}
+	}
+
+	bool AnyActiveSettingsEditor()
+	{
+		return EditorWidgets.Num() > 0;
+	}
+
+private:
+
+	void RegisterAutoDiscoveredSettings(ISettingsModule& SettingsModule)
+	{
+		// Find game object
+		for ( TObjectIterator<UDeveloperSettings> SettingsIt(RF_NoFlags); SettingsIt; ++SettingsIt )
+		{
+			if ( UDeveloperSettings* Settings = *SettingsIt )
+			{
+				// Only Add the CDO of any UDeveloperSettings objects.
+				if ( Settings->HasAnyFlags(RF_ClassDefaultObject) && !Settings->GetClass()->HasAnyCastFlag(CLASS_Deprecated) )
+				{
+					// Ignore the setting if it's specifically the UDeveloperSettings
+					if ( Settings->GetClass() == UDeveloperSettings::StaticClass() )
+					{
+						continue;
+					}
+
+					FRegisteredSettings Registered;
+					Registered.ContainerName = Settings->GetContainerName();
+					Registered.CategoryName = Settings->GetCategoryName();
+					Registered.SectionName = Settings->GetSectionName();
+
+					TSharedPtr<SWidget> CustomWidget = Settings->GetCustomSettingsWidget();
+					if ( CustomWidget.IsValid() )
+					{
+						// Add Settings
+						SettingsModule.RegisterSettings(Registered.ContainerName, Registered.CategoryName, Registered.SectionName,
+							Settings->GetSectionText(),
+							Settings->GetSectionDescription(),
+							CustomWidget.ToSharedRef()
+							);
+					}
+					else
+					{
+						// Add Settings
+						SettingsModule.RegisterSettings(Registered.ContainerName, Registered.CategoryName, Registered.SectionName,
+							Settings->GetSectionText(),
+							Settings->GetSectionDescription(),
+							Settings
+							);
+					}
+
+					AutoDiscoveredSettings.Add(Registered);
+				}
+			}
+		}
+	}
+
+	void UnregisterAutoDiscoveredSettings(ISettingsModule& SettingsModule)
+	{
+		// Unregister any auto discovers settings.
+		for ( const FRegisteredSettings& Settings : AutoDiscoveredSettings )
+		{
+			SettingsModule.UnregisterSettings(Settings.ContainerName, Settings.CategoryName, Settings.SectionName);
+		}
+
+		AutoDiscoveredSettings.Reset();
+	}
+
 private:
 
 	FApplicationRestartRequiredNotification ApplicationRestartRequiredNotification;
+
+	/** The list of auto discovered settings that need to be unregistered. */
+	TArray<FRegisteredSettings> AutoDiscoveredSettings;
+
+	/** Living editor widgets that have been handed out. */
+	TArray< TWeakPtr<SWidget> > EditorWidgets;
+
+	/** Flag if the settings are stale currently and need to be refreshed. */
+	bool bAreSettingsStale;
 };
 
 
diff --git a/Engine/Source/Developer/SettingsEditor/Private/Widgets/SSettingsEditor.cpp b/Engine/Source/Developer/SettingsEditor/Private/Widgets/SSettingsEditor.cpp
--- a/Engine/Source/Developer/SettingsEditor/Private/Widgets/SSettingsEditor.cpp
+++ b/Engine/Source/Developer/SettingsEditor/Private/Widgets/SSettingsEditor.cpp
@@ -31,6 +31,7 @@ SSettingsEditor::~SSettingsEditor()
 
 void SSettingsEditor::Construct( const FArguments& InArgs, const ISettingsEditorModelRef& InModel )
 {
+	bIsActiveTimerRegistered = false;
 	Model = InModel;
 	SettingsContainer = InModel->GetSettingsContainer();
 	OnApplicationRestartRequiredDelegate = InArgs._OnApplicationRestartRequired;
@@ -834,9 +835,21 @@ EVisibility SSettingsEditor::HandleSettingsBoxVisibility() const
 
 void SSettingsEditor::HandleSettingsContainerCategoryModified( const FName& CategoryName )
 {
-	ReloadCategories();
+	if ( !bIsActiveTimerRegistered )
+	{
+		bIsActiveTimerRegistered = true;
+		RegisterActiveTimer(0.f, FWidgetActiveTimerDelegate::CreateSP(this, &SSettingsEditor::UpdateCategoriesCallback));
+	}
 }
 
+EActiveTimerReturnType SSettingsEditor::UpdateCategoriesCallback(double InCurrentTime, float InDeltaTime)
+{
+	bIsActiveTimerRegistered = false;
+
+	ReloadCategories();
+
+	return EActiveTimerReturnType::Stop;
+}
 
 bool SSettingsEditor::HandleSettingsViewEnabled() const
 {
diff --git a/Engine/Source/Developer/SettingsEditor/Private/Widgets/SSettingsEditor.h b/Engine/Source/Developer/SettingsEditor/Private/Widgets/SSettingsEditor.h
--- a/Engine/Source/Developer/SettingsEditor/Private/Widgets/SSettingsEditor.h
+++ b/Engine/Source/Developer/SettingsEditor/Private/Widgets/SSettingsEditor.h
@@ -184,6 +184,9 @@ private:
 	/** Callback for determining the visibility of the settings view. */
 	EVisibility HandleSettingsViewVisibility() const;
 
+	/** Callback when the timer has been ticked to refresh the categories latently. */
+	EActiveTimerReturnType UpdateCategoriesCallback(double InCurrentTime, float InDeltaTime);
+
 private:
 
 	/** Holds the vertical box for settings categories. */
@@ -209,4 +212,7 @@ private:
 
 	/** Delegate called when this settings editor requests that the user be notified that the application needs to be restarted for some setting changes to take effect */
 	FSimpleDelegate OnApplicationRestartRequiredDelegate;
+
+	/** Is the active timer registered to refresh categories after the settings changed. */
+	bool bIsActiveTimerRegistered;
 };
diff --git a/Engine/Source/Editor/ProjectSettingsViewer/Private/ProjectSettingsViewerModule.cpp b/Engine/Source/Editor/ProjectSettingsViewer/Private/ProjectSettingsViewerModule.cpp
--- a/Engine/Source/Editor/ProjectSettingsViewer/Private/ProjectSettingsViewerModule.cpp
+++ b/Engine/Source/Editor/ProjectSettingsViewer/Private/ProjectSettingsViewerModule.cpp
@@ -47,16 +47,9 @@ struct FRegisteredSettings
 class FProjectSettingsViewerModule
 	: public IModuleInterface
 	, public ISettingsViewer
-	, public FTickableEditorObject
 {
 public:
 
-	FProjectSettingsViewerModule()
-	{
-		bTicking = false;
-		bInvalidated = false;
-	}
-
 	// ISettingsViewer interface
 
 	virtual void ShowSettings( const FName& CategoryName, const FName& SectionName ) override
@@ -87,9 +80,6 @@ public:
 		{
 			RegisterEngineSettings(*SettingsModule);
 			RegisterProjectSettings(*SettingsModule);
-			RegisterEditorSettings(*SettingsModule);
-
-			RegisterAutoDiscoveredSettings(*SettingsModule);
 
 			SettingsModule->RegisterViewer("Project", *this);
 		}
@@ -97,30 +87,14 @@ public:
 		FGlobalTabmanager::Get()->RegisterNomadTabSpawner(ProjectSettingsTabName, FOnSpawnTab::CreateRaw(this, &FProjectSettingsViewerModule::HandleSpawnSettingsTab))
 			.SetDisplayName(LOCTEXT("ProjectSettingsTabTitle", "Project Settings"))
 			.SetMenuType(ETabSpawnerMenuType::Hidden);
-
-		FModuleManager::Get().OnModulesChanged().AddRaw(this, &FProjectSettingsViewerModule::ModulesChangesCallback);
-
-		bTicking = true;
 	}
 
 	virtual void ShutdownModule() override
 	{
-		bTicking = false;
-
 		FGlobalTabmanager::Get()->UnregisterNomadTabSpawner(ProjectSettingsTabName);
 		UnregisterSettings();
-
-		FModuleManager::Get().OnModulesChanged().RemoveAll(this);
 	}
-
-	void ModulesChangesCallback(FName ModuleName, EModuleChangeReason ReasonForChange)
-	{
-		if ( ReasonForChange == EModuleChangeReason::ModuleLoaded )
-		{
-			bInvalidated = true;
-		}
-	}
-
+	
 	virtual bool SupportsDynamicReloading() override
 	{
 		return true;
@@ -142,20 +116,6 @@ protected:
 			GetMutableDefault<UEngine>()
 		);
 
-		// audio settings
-		SettingsModule.RegisterSettings("Project", "Engine", "Audio",
-			LOCTEXT("EngineAudioSettingsName", "Audio"),
-			LOCTEXT("ProjectAudioSettingsDescription", "Audio settings."),
-			GetMutableDefault<UAudioSettings>()
-		);
-
-		// collision settings
-		SettingsModule.RegisterSettings("Project", "Engine", "Collision",
-			LOCTEXT("ProjectCollisionSettingsName", "Collision"),
-			LOCTEXT("ProjectCollisionSettingsDescription", "Set up and modify collision settings."),
-			GetMutableDefault<UCollisionProfile>()
-		);
-
 		// command console settings
 		SettingsModule.RegisterSettings("Project", "Engine", "Console",
 			LOCTEXT("ProjectConsoleSettingsName", "Console"),
@@ -200,48 +160,6 @@ protected:
 			LOCTEXT("CrowdManagerSettingsDescription", "Settings for the AI Crowd Manager."),
 			GetMutableDefault<UCrowdManager>()
 			);
-
-		// Physics settings
-		SettingsModule.RegisterSettings("Project", "Engine", "Animation",
-			LOCTEXT("EngineAnimationSettingsName", "Animation"),
-			LOCTEXT("ProjectAnimationSettingsDescription", "Default animation settings."),
-			GetMutableDefault<UAnimationSettings>()
-		);
-
-		// network settings
-		SettingsModule.RegisterSettings("Project", "Engine", "Network",
-			LOCTEXT("NetworkSettingsName", "Network"),
-			LOCTEXT("NetworkSettingsDescription", "Network settings."),
-			GetMutableDefault<UNetworkSettings>()
-		);
-
-		// Physics settings
-		SettingsModule.RegisterSettings("Project", "Engine", "Physics",
-			LOCTEXT("EnginePhysicsSettingsName", "Physics"),
-			LOCTEXT("ProjectPhysicsSettingsDescription", "Default physics settings."),
-			GetMutableDefault<UPhysicsSettings>()
-		);
-
-		// Rendering settings
-		SettingsModule.RegisterSettings("Project", "Engine", "Rendering",
-			LOCTEXT("EngineRenderingSettingsName", "Rendering"),
-			LOCTEXT("ProjectRenderingSettingsDescription", "Rendering settings."),
-			GetMutableDefault<URendererSettings>()
-		);
-
-		// UI settings
-		SettingsModule.RegisterSettings("Project", "Engine", "UI",
-			LOCTEXT("EngineUISettingsName", "User Interface"),
-			LOCTEXT("ProjectUISettingsDescription", "User Interface settings that control Slate and UMG."),
-			GetMutableDefault<UUserInterfaceSettings>()
-		);
-
-		// Cooker settings
-		SettingsModule.RegisterSettings("Project", "Engine", "Cooker",
-			LOCTEXT("CookerSettingsName", "Cooker"),
-			LOCTEXT("CookerSettingsDescription", "Various cooker settings."),
-			GetMutableDefault<UCookerSettings>()
-			);
 	}
 
 	/**
@@ -288,69 +206,6 @@ protected:
 		);
 	}
 
-	void RegisterAutoDiscoveredSettings(ISettingsModule& SettingsModule)
-	{
-		// Find game object
-		for ( TObjectIterator<UDeveloperSettings> SettingsIt(RF_NoFlags); SettingsIt; ++SettingsIt )
-		{
-			if ( UDeveloperSettings* Settings = *SettingsIt )
-			{
-				// Only Add the CDO of any UDeveloperSettings objects.
-				if ( Settings->HasAnyFlags(RF_ClassDefaultObject) && !Settings->GetClass()->HasAnyCastFlag(CLASS_Deprecated | CLASS_Abstract) )
-				{
-					FRegisteredSettings Registered;
-					Registered.ContainerName = Settings->GetContainerName();
-					Registered.CategoryName = Settings->GetCategoryName();
-					Registered.SectionName = Settings->GetSectionName();
-
-					TSharedPtr<SWidget> CustomWidget = Settings->GetCustomSettingsWidget();
-					if ( CustomWidget.IsValid() )
-					{
-						// Add Settings
-						SettingsModule.RegisterSettings(Registered.ContainerName, Registered.CategoryName, Registered.SectionName,
-							Settings->GetSectionText(),
-							Settings->GetSectionDescription(),
-							CustomWidget.ToSharedRef()
-							);
-					}
-					else
-					{
-						// Add Settings
-						SettingsModule.RegisterSettings(Registered.ContainerName, Registered.CategoryName, Registered.SectionName,
-							Settings->GetSectionText(),
-							Settings->GetSectionDescription(),
-							Settings
-							);
-					}
-
-					AutoDiscoveredSettings.Add(Registered);
-				}
-			}
-		}
-	}
-
-	/**
-	 * Registers Editor settings.
-	 *
-	 * @param SettingsModule A reference to the settings module.
-	 */
-	void RegisterEditorSettings( ISettingsModule& SettingsModule )
-	{
-		// Appearance settings
-		SettingsModule.RegisterSettings("Project", "Editor", "Appearance",
-			LOCTEXT("AppearanceSettingsName", "Appearance"),
-			LOCTEXT("AppearanceSettingsDescription", "Settings pertaining to the appearance of the editor"),
-			GetMutableDefault<UEditorProjectAppearanceSettings>()
-		);
-
-		// view port settings
-		SettingsModule.RegisterSettings("Project", "Editor", "2D",
-			LOCTEXT("Editor2DSettingsName", "2D"),
-			LOCTEXT("Editor2DSettingsDescription", "Configure the settings for the 2D Level Editor."),
-			GetMutableDefault<ULevelEditor2DSettings>()
-		);
-	}
-
 	/** Unregisters all settings. */
 	void UnregisterSettings()
 	{
@@ -379,53 +234,9 @@ protected:
 
 			// Editor settings
 			SettingsModule->UnregisterSettings("Editor", "Editor", "Appearance");
-
-			UnregisterAutoDiscoveredSettings(*SettingsModule);
 		}
 	}
 
-	void UnregisterAutoDiscoveredSettings(ISettingsModule& SettingsModule)
-	{
-		// Unregister any auto discovers settings.
-		for ( const FRegisteredSettings& Settings : AutoDiscoveredSettings )
-		{
-			SettingsModule.UnregisterSettings(Settings.ContainerName, Settings.CategoryName, Settings.SectionName);
-		}
-
-		AutoDiscoveredSettings.Reset();
-	}
-
-private:
-
-	/** FTickableEditorObject interface */
-	void Tick(float DeltaTime) override
-	{
-		if ( bInvalidated )
-		{
-			ISettingsModule* SettingsModule = FModuleManager::GetModulePtr<ISettingsModule>("Settings");
-
-			if ( SettingsModule != nullptr )
-			{
-				UnregisterAutoDiscoveredSettings(*SettingsModule);
-				RegisterAutoDiscoveredSettings(*SettingsModule);
-			}
-
-			bInvalidated = false;
-		}
-	}
-
-	bool IsTickable() const override
-	{
-		return bTicking;
-	}
-
-	TStatId GetStatId() const override
-	{
-		RETURN_QUICK_DECLARE_CYCLE_STAT(FProjectSettingsViewerModule, STATGROUP_Tickables);
-	}
-
-	/** FTickableEditorObject interface */
-
 private:
 
 	/** Handles creating the project settings tab. */
@@ -459,15 +270,6 @@ private:
 
 	/** Holds a pointer to the settings editor's view model. */
 	TWeakPtr<ISettingsEditorModel> SettingsEditorModelPtr;
-
-	/** The list of auto discovered settings that need to be unregistered. */
-	TArray<FRegisteredSettings> AutoDiscoveredSettings;
-
-	/** When new modules are loaded we invalidate the list of autodiscovered settings and re-discover them. */
-	bool bInvalidated;
-
-	/** Should we be ticking? */
-	bool bTicking;
 };
 
 
diff --git a/Engine/Source/Editor/UnrealEd/Private/Settings/EditorProjectSettings.cpp b/Engine/Source/Editor/UnrealEd/Private/Settings/EditorProjectSettings.cpp
--- a/Engine/Source/Editor/UnrealEd/Private/Settings/EditorProjectSettings.cpp
+++ b/Engine/Source/Editor/UnrealEd/Private/Settings/EditorProjectSettings.cpp
@@ -44,8 +44,9 @@ TArray<EUnit> ToRawUnits(const TArray<TEnumAsByte<EUnit>>& In)
 	return Units;	
 }
 
-UEditorProjectAppearanceSettings::UEditorProjectAppearanceSettings(const FObjectInitializer&)
-	: UnitDisplay_DEPRECATED(EUnitDisplay::Invalid)
+UEditorProjectAppearanceSettings::UEditorProjectAppearanceSettings(const FObjectInitializer& Initializer)
+	: Super(Initializer)
+	, UnitDisplay_DEPRECATED(EUnitDisplay::Invalid)
 	, DefaultInputUnits_DEPRECATED(EDefaultLocationUnit::Invalid)
 {
 }
diff --git a/Engine/Source/Editor/UnrealEd/Public/CookerSettings.h b/Engine/Source/Editor/UnrealEd/Public/CookerSettings.h
--- a/Engine/Source/Editor/UnrealEd/Public/CookerSettings.h
+++ b/Engine/Source/Editor/UnrealEd/Public/CookerSettings.h
@@ -9,11 +9,10 @@
 #include "CookerSettings.generated.h"
 
 /**
- * Implements per-project cooker settings exposed to the editor
+ * Various cooker settings.
  */
-UCLASS(config=Engine, defaultconfig)
-class UNREALED_API UCookerSettings
-	: public UObject
+UCLASS(config=Engine, defaultconfig, meta=(DisplayName="Cooker"))
+class UNREALED_API UCookerSettings : public UDeveloperSettings
 {
 	GENERATED_UCLASS_BODY()
 
diff --git a/Engine/Source/Editor/UnrealEd/Public/Settings/EditorProjectSettings.h b/Engine/Source/Editor/UnrealEd/Public/Settings/EditorProjectSettings.h
--- a/Engine/Source/Editor/UnrealEd/Public/Settings/EditorProjectSettings.h
+++ b/Engine/Source/Editor/UnrealEd/Public/Settings/EditorProjectSettings.h
@@ -6,9 +6,11 @@
 
 #include "EditorProjectSettings.generated.h"
 
-/** Editor project appearance settings. Stored in default config, per-project */
-UCLASS(config=Editor, defaultconfig)
-class UNREALED_API UEditorProjectAppearanceSettings : public UObject
+/**
+ * Editor project appearance settings. Stored in default config, per-project
+ */
+UCLASS(config=Editor, defaultconfig, meta=(DisplayName="Appearance"))
+class UNREALED_API UEditorProjectAppearanceSettings : public UDeveloperSettings
 {
 public:
 	GENERATED_BODY()
@@ -100,10 +102,10 @@ enum class ELevelEditor2DAxis : uint8
 };
 
 /**
-* Implements the Project specific 2D viewport settings.
-*/
-UCLASS(config=Editor)
-class UNREALED_API ULevelEditor2DSettings : public UObject
+ * Configure the settings for the 2D Level Editor
+ */
+UCLASS(config=Editor, meta=(DisplayName="2D"))
+class UNREALED_API ULevelEditor2DSettings : public UDeveloperSettings
 {
 	GENERATED_UCLASS_BODY()
 
diff --git a/Engine/Source/Runtime/Engine/Classes/Animation/AnimationSettings.h b/Engine/Source/Runtime/Engine/Classes/Animation/AnimationSettings.h
--- a/Engine/Source/Runtime/Engine/Classes/Animation/AnimationSettings.h
+++ b/Engine/Source/Runtime/Engine/Classes/Animation/AnimationSettings.h
@@ -10,11 +10,10 @@
 #include "AnimationSettings.generated.h"
 
 /**
- * Implements project settings for the animation settings
+ * Default animation settings.
  */
-UCLASS(config=Engine, defaultconfig)
-class ENGINE_API UAnimationSettings
-	: public UObject
+UCLASS(config=Engine, defaultconfig, meta=(DisplayName="Animation"))
+class ENGINE_API UAnimationSettings : public UDeveloperSettings
 {
 	GENERATED_UCLASS_BODY()
 
diff --git a/Engine/Source/Runtime/Engine/Classes/Engine/CollisionProfile.h b/Engine/Source/Runtime/Engine/Classes/Engine/CollisionProfile.h
--- a/Engine/Source/Runtime/Engine/Classes/Engine/CollisionProfile.h
+++ b/Engine/Source/Runtime/Engine/Classes/Engine/CollisionProfile.h
@@ -2,6 +2,8 @@
 
 #pragma once
 
+#include "Engine/DeveloperSettings.h"
+
 /**
  *	This will hold all of our enums and types and such that we need to
  *	use in multiple files where the enum can't be mapped to a specific file.
@@ -134,11 +136,10 @@ struct ENGINE_API FCustomProfile
 
 
 /**
- * Implements a collision profile for the collision sub-system.
+ * Set up and modify collision settings.
  */
-UCLASS(abstract, config=Engine, defaultconfig, MinimalAPI)
-class UCollisionProfile
-	: public UObject
+UCLASS(abstract, config=Engine, defaultconfig, MinimalAPI, meta=(DisplayName="Collision"))
+class UCollisionProfile : public UDeveloperSettings
 {
 	GENERATED_UCLASS_BODY()
 
diff --git a/Engine/Source/Runtime/Engine/Classes/Engine/DeveloperSettings.h b/Engine/Source/Runtime/Engine/Classes/Engine/DeveloperSettings.h
--- a/Engine/Source/Runtime/Engine/Classes/Engine/DeveloperSettings.h
+++ b/Engine/Source/Runtime/Engine/Classes/Engine/DeveloperSettings.h
@@ -15,7 +15,7 @@ enum class EDeveloperSettingsContainer : uint8
 /**
  * The base class of any auto discovered settings object.
  */
-UCLASS(abstract)
+UCLASS(Abstract)
 class ENGINE_API UDeveloperSettings : public UObject
 {
 	GENERATED_BODY()
@@ -23,6 +23,8 @@ class ENGINE_API UDeveloperSettings : public UObject
 public:
 	UDeveloperSettings(const FObjectInitializer& ObjectInitializer);
 
+	/** Gets the container type, Project or Editor */
+	EDeveloperSettingsContainer GetContainerType() const { return Container; }
 	/** Gets the settings container name for the settings, either Project or Editor */
 	FName GetContainerName() const;
 	/** Gets the category for the settings, some high level grouping like, Editor, Engine, Game...etc. */
@@ -38,7 +40,7 @@ public:
 #endif
 
 	/** Gets a custom widget for the settings.  This is only for very custom situations. */
-	virtual TSharedPtr<SWidget> GetCustomSettingsWidget() const;
+	virtual TSharedPtr<class SWidget> GetCustomSettingsWidget() const;
 
 private:
 	/**
diff --git a/Engine/Source/Runtime/Engine/Classes/Engine/NetworkSettings.h b/Engine/Source/Runtime/Engine/Classes/Engine/NetworkSettings.h
--- a/Engine/Source/Runtime/Engine/Classes/Engine/NetworkSettings.h
+++ b/Engine/Source/Runtime/Engine/Classes/Engine/NetworkSettings.h
@@ -6,11 +6,10 @@
 
 
 /**
- * Implements project settings for the Network sub-system.
+ * Network settings.
  */
-UCLASS(config=Engine, defaultconfig)
-class ENGINE_API UNetworkSettings
-	: public UObject
+UCLASS(config=Engine, defaultconfig, meta=(DisplayName="Network"))
+class ENGINE_API UNetworkSettings : public UDeveloperSettings
 {
 	GENERATED_UCLASS_BODY()
 
diff --git a/Engine/Source/Runtime/Engine/Classes/Engine/RendererSettings.h b/Engine/Source/Runtime/Engine/Classes/Engine/RendererSettings.h
--- a/Engine/Source/Runtime/Engine/Classes/Engine/RendererSettings.h
+++ b/Engine/Source/Runtime/Engine/Classes/Engine/RendererSettings.h
@@ -85,11 +85,10 @@ namespace EAntiAliasingMethodUI
 
 
 /**
- * Implements project settings for the Rendering sub-system.
+ * Rendering settings.
  */
-UCLASS(config=Engine, defaultconfig)
-class ENGINE_API URendererSettings
-	: public UObject
+UCLASS(config=Engine, defaultconfig, meta=(DisplayName="Rendering"))
+class ENGINE_API URendererSettings : public UDeveloperSettings
 {
 	GENERATED_UCLASS_BODY()
 
diff --git a/Engine/Source/Runtime/Engine/Classes/Engine/UserInterfaceSettings.h b/Engine/Source/Runtime/Engine/Classes/Engine/UserInterfaceSettings.h
--- a/Engine/Source/Runtime/Engine/Classes/Engine/UserInterfaceSettings.h
+++ b/Engine/Source/Runtime/Engine/Classes/Engine/UserInterfaceSettings.h
@@ -1,6 +1,7 @@
 // Copyright 1998-2015 Epic Games, Inc. All Rights Reserved.
 
 #pragma once
+
 #include "Curves/CurveFloat.h"
 #include "UserInterfaceSettings.generated.h"
 
@@ -38,11 +39,10 @@ class UWidget;
 class UDPICustomScalingRule;
 
 /**
- * Implements user interface related settings.
+ * User Interface settings that control Slate and UMG.
  */
-UCLASS(config=Engine, defaultconfig)
-class ENGINE_API UUserInterfaceSettings
-	: public UObject
+UCLASS(config=Engine, defaultconfig, meta=(DisplayName="User Interface"))
+class ENGINE_API UUserInterfaceSettings : public UDeveloperSettings
 {
 	GENERATED_UCLASS_BODY()
 
diff --git a/Engine/Source/Runtime/Engine/Classes/PhysicsEngine/PhysicsSettings.h b/Engine/Source/Runtime/Engine/Classes/PhysicsEngine/PhysicsSettings.h
--- a/Engine/Source/Runtime/Engine/Classes/PhysicsEngine/PhysicsSettings.h
+++ b/Engine/Source/Runtime/Engine/Classes/PhysicsEngine/PhysicsSettings.h
@@ -84,11 +84,10 @@ namespace ESettingsLockedAxis
 
 
 /**
- * Implements project settings for the physics sub-system.
+ * Default physics settings.
  */
-UCLASS(config=Engine, defaultconfig)
-class ENGINE_API UPhysicsSettings
-	: public UObject
+UCLASS(config=Engine, defaultconfig, meta=(DisplayName="Physics"))
+class ENGINE_API UPhysicsSettings : public UDeveloperSettings
 {
 	GENERATED_UCLASS_BODY()
 
diff --git a/Engine/Source/Runtime/Engine/Classes/Sound/AudioSettings.h b/Engine/Source/Runtime/Engine/Classes/Sound/AudioSettings.h
--- a/Engine/Source/Runtime/Engine/Classes/Sound/AudioSettings.h
+++ b/Engine/Source/Runtime/Engine/Classes/Sound/AudioSettings.h
@@ -6,11 +6,10 @@
 
 
 /**
- * Implements audio related settings.
+ * Audio settings.
  */
-UCLASS(config=Engine, defaultconfig)
-class ENGINE_API UAudioSettings
-	: public UObject
+UCLASS(config=Engine, defaultconfig, meta=(DisplayName="Audio"))
+class ENGINE_API UAudioSettings : public UDeveloperSettings
 {
 	GENERATED_UCLASS_BODY()
 
diff --git a/Engine/Source/Runtime/Engine/Private/EnginePrivate.h b/Engine/Source/Runtime/Engine/Private/EnginePrivate.h
--- a/Engine/Source/Runtime/Engine/Private/EnginePrivate.h
+++ b/Engine/Source/Runtime/Engine/Private/EnginePrivate.h
@@ -7,6 +7,7 @@ EnginePrivate.h: Unreal engine private header file.
 #pragma once
 #include "CoreUObject.h"
 #include "EngineDefines.h"
+#include "Engine/DeveloperSettings.h"
 #include "GameMapsSettings.h"
 #include "EngineStats.h"
 #include "EngineLogs.h"
diff --git a/Engine/Source/Runtime/Engine/Public/Engine.h b/Engine/Source/Runtime/Engine/Public/Engine.h
--- a/Engine/Source/Runtime/Engine/Public/Engine.h
+++ b/Engine/Source/Runtime/Engine/Public/Engine.h
@@ -22,6 +22,7 @@
 #include "EngineMinimal.h" // 
 
 #include "Engine/EngineBaseTypes.h"
+#include "Engine/DeveloperSettings.h"
 #include "Camera/CameraTypes.h"
 #include "Engine/EngineTypes.h"
 #include "Sound/AmbientSound.h"
